WinArm64 version of UEWavComp

#rnx

diff --git a/Engine/Source/Runtime/UEWavComp/UEWavComp.Build.cs b/Engine/Source/Runtime/UEWavComp/UEWavComp.Build.cs
--- a/Engine/Source/Runtime/UEWavComp/UEWavComp.Build.cs
+++ b/Engine/Source/Runtime/UEWavComp/UEWavComp.Build.cs
@@ -15,6 +15,12 @@ public class UEWavComp : ModuleRules
 			PublicAdditionalLibraries.Add(Path.Combine(LibPath, "uewav_w64.lib"));
 			bWithLib = true;
 		}
+		else if (Target.Platform == UnrealTargetPlatform.Win64 && Target.WindowsPlatform.Architecture == UnrealArch.Arm64)
+		{
+			string LibPath = Path.Combine(ModuleDirectory, "Lib");
+			PublicAdditionalLibraries.Add(Path.Combine(LibPath, "uewav_warm64.lib"));
+			bWithLib = true;
+		}
 		else if (Target.Platform == UnrealTargetPlatform.Mac)
 		{
 			string LibPath = Path.Combine(ModuleDirectory, "Lib");
